TLazyObjectPtr Deprecation:
Change usage in LevelSequence, SequenceRecorder, TakeRecorder plugins to TWeakObjectPtr equivalent
Delete long unsued LeveSequenceObject class

#Synced-CL 52628728

diff --git a/Engine/Plugins/FX/NiagaraSimCaching/Source/NiagaraSimCachingEditor/Public/Sequencer/MovieSceneNiagaraTrackRecorder.h b/Engine/Plugins/FX/NiagaraSimCaching/Source/NiagaraSimCachingEditor/Public/Sequencer/MovieSceneNiagaraTrackRecorder.h
--- a/Engine/Plugins/FX/NiagaraSimCaching/Source/NiagaraSimCachingEditor/Public/Sequencer/MovieSceneNiagaraTrackRecorder.h
+++ b/Engine/Plugins/FX/NiagaraSimCaching/Source/NiagaraSimCachingEditor/Public/Sequencer/MovieSceneNiagaraTrackRecorder.h
@@ -66,7 +66,7 @@ private:
 	TWeakObjectPtr<UMovieSceneNiagaraCacheSection> NiagaraCacheSection;
 
 	/** Object to record from */
-	TLazyObjectPtr<UNiagaraComponent> SystemToRecord;
+	TWeakObjectPtr<UNiagaraComponent> SystemToRecord;
 
 	TOptional<FFrameNumberRange> RecordRange;
 
diff --git a/Engine/Plugins/MovieScene/LevelSequenceEditor/Source/LevelSequenceEditor/Public/Misc/LevelSequenceEditorSettings.h b/Engine/Plugins/MovieScene/LevelSequenceEditor/Source/LevelSequenceEditor/Public/Misc/LevelSequenceEditorSettings.h
--- a/Engine/Plugins/MovieScene/LevelSequenceEditor/Source/LevelSequenceEditor/Public/Misc/LevelSequenceEditorSettings.h
+++ b/Engine/Plugins/MovieScene/LevelSequenceEditor/Source/LevelSequenceEditor/Public/Misc/LevelSequenceEditorSettings.h
@@ -102,7 +102,7 @@ public:
 
 	/** Sequence With Shots level sequence to duplicate when creating shots. */
 	UPROPERTY(Transient, DisplayName="Sequence to Duplicate", EditAnywhere, Category=SequenceWithShots)
-	TLazyObjectPtr<class ULevelSequence> SequenceToDuplicate;
+	TSoftObjectPtr<class ULevelSequence> SequenceToDuplicate;
 
 	/** Array of sub sequence names, each will result in a level sequence asset in the shot. */
 	UPROPERTY(EditAnywhere, BlueprintReadWrite, Category= SequenceWithShots)
diff --git a/Engine/Plugins/VirtualProduction/Takes/Source/TakeTrackRecorders/Public/TrackRecorders/MovieScene3DAttachTrackRecorder.h b/Engine/Plugins/VirtualProduction/Takes/Source/TakeTrackRecorders/Public/TrackRecorders/MovieScene3DAttachTrackRecorder.h
--- a/Engine/Plugins/VirtualProduction/Takes/Source/TakeTrackRecorders/Public/TrackRecorders/MovieScene3DAttachTrackRecorder.h
+++ b/Engine/Plugins/VirtualProduction/Takes/Source/TakeTrackRecorders/Public/TrackRecorders/MovieScene3DAttachTrackRecorder.h
@@ -50,7 +50,7 @@ private:
 	TWeakObjectPtr<class UMovieScene3DAttachTrack> AttachTrack;
 
 	/** Track the actor we are attached to */
-	TLazyObjectPtr<class AActor> ActorAttachedTo;
+	TWeakObjectPtr<class AActor> ActorAttachedTo;
 };
 
 #undef UE_API
diff --git a/Engine/Plugins/VirtualProduction/Takes/Source/TakeTrackRecorders/Public/TrackRecorders/MovieSceneParticleTrackRecorder.h b/Engine/Plugins/VirtualProduction/Takes/Source/TakeTrackRecorders/Public/TrackRecorders/MovieSceneParticleTrackRecorder.h
--- a/Engine/Plugins/VirtualProduction/Takes/Source/TakeTrackRecorders/Public/TrackRecorders/MovieSceneParticleTrackRecorder.h
+++ b/Engine/Plugins/VirtualProduction/Takes/Source/TakeTrackRecorders/Public/TrackRecorders/MovieSceneParticleTrackRecorder.h
@@ -54,7 +54,7 @@ private:
 	UE_API void OnTriggered(UParticleSystemComponent* Component, bool bActivating);
 private:
 	/** Object to record from */
-	TLazyObjectPtr<class UParticleSystemComponent> SystemToRecord;
+	TWeakObjectPtr<class UParticleSystemComponent> SystemToRecord;
 
 	/** Section to record to */
 	TWeakObjectPtr<class UMovieSceneParticleSection> MovieSceneSection;
diff --git a/Engine/Plugins/VirtualProduction/Takes/Source/TakeTrackRecorders/Public/TrackRecorders/MovieSceneTrackRecorder.h b/Engine/Plugins/VirtualProduction/Takes/Source/TakeTrackRecorders/Public/TrackRecorders/MovieSceneTrackRecorder.h
--- a/Engine/Plugins/VirtualProduction/Takes/Source/TakeTrackRecorders/Public/TrackRecorders/MovieSceneTrackRecorder.h
+++ b/Engine/Plugins/VirtualProduction/Takes/Source/TakeTrackRecorders/Public/TrackRecorders/MovieSceneTrackRecorder.h
@@ -176,7 +176,7 @@ protected:
 
 protected:
 	/** Object to record from */
-	TLazyObjectPtr<UObject> ObjectToRecord;
+	TWeakObjectPtr<UObject> ObjectToRecord;
 
 	/** Owning Object GUID in the Level Sequence */
 	FGuid ObjectGuid;
diff --git a/Engine/Source/Editor/SequenceRecorder/Private/Sections/MovieScene3DTransformSectionRecorder.h b/Engine/Source/Editor/SequenceRecorder/Private/Sections/MovieScene3DTransformSectionRecorder.h
--- a/Engine/Source/Editor/SequenceRecorder/Private/Sections/MovieScene3DTransformSectionRecorder.h
+++ b/Engine/Source/Editor/SequenceRecorder/Private/Sections/MovieScene3DTransformSectionRecorder.h
@@ -81,7 +81,7 @@ private:
 
 private:
 	/** Object to record from */
-	TLazyObjectPtr<UObject> ObjectToRecord;
+	TWeakObjectPtr<UObject> ObjectToRecord;
 
 	/** MovieScene to record to */
 	TWeakObjectPtr<class UMovieScene> MovieScene;
diff --git a/Engine/Source/Editor/SequenceRecorder/Private/Sections/MovieSceneAnimationSectionRecorder.h b/Engine/Source/Editor/SequenceRecorder/Private/Sections/MovieSceneAnimationSectionRecorder.h
--- a/Engine/Source/Editor/SequenceRecorder/Private/Sections/MovieSceneAnimationSectionRecorder.h
+++ b/Engine/Source/Editor/SequenceRecorder/Private/Sections/MovieSceneAnimationSectionRecorder.h
@@ -60,7 +60,7 @@ public:
 
 private:
 	/** Object to record from */
-	TLazyObjectPtr<UObject> ObjectToRecord;
+	TWeakObjectPtr<UObject> ObjectToRecord;
 
 	/** Section to record to */
 	TWeakObjectPtr<class UMovieSceneSkeletalAnimationSection> MovieSceneSection;
diff --git a/Engine/Source/Editor/SequenceRecorder/Private/Sections/MovieSceneMultiPropertyRecorder.h b/Engine/Source/Editor/SequenceRecorder/Private/Sections/MovieSceneMultiPropertyRecorder.h
--- a/Engine/Source/Editor/SequenceRecorder/Private/Sections/MovieSceneMultiPropertyRecorder.h
+++ b/Engine/Source/Editor/SequenceRecorder/Private/Sections/MovieSceneMultiPropertyRecorder.h
@@ -6,7 +6,7 @@
 #include "IMovieSceneSectionRecorder.h"
 #include "IMovieSceneSectionRecorderFactory.h"
 #include "Templates/SharedPointer.h"
-#include "UObject/LazyObjectPtr.h"
+#include "UObject/WeakObjectPtr.h"
 #include "UObject/Object.h"
 
 class FProperty;
@@ -45,7 +45,7 @@ public:
 
 private:
 	/** Object to record from */
-	TLazyObjectPtr<class UObject> ObjectToRecord;
+	TWeakObjectPtr<class UObject> ObjectToRecord;
 
 	/** All of our property recorders */
 	TArray<TSharedPtr<class IMovieScenePropertyRecorder>> PropertyRecorders;
diff --git a/Engine/Source/Editor/SequenceRecorder/Private/SequenceRecorder.cpp b/Engine/Source/Editor/SequenceRecorder/Private/SequenceRecorder.cpp
--- a/Engine/Source/Editor/SequenceRecorder/Private/SequenceRecorder.cpp
+++ b/Engine/Source/Editor/SequenceRecorder/Private/SequenceRecorder.cpp
@@ -682,7 +682,7 @@ bool FSequenceRecorder::StartRecordingInternal(UWorld* World)
 	}
 
 	// kick off level sequence actors we are syncing to
-	for(TLazyObjectPtr<ALevelSequenceActor> LevelSequenceActor : Settings->LevelSequenceActorsToTrigger)
+	for(TWeakObjectPtr<ALevelSequenceActor> LevelSequenceActor : Settings->LevelSequenceActorsToTrigger)
 	{
 		if(LevelSequenceActor.IsValid())
 		{
diff --git a/Engine/Source/Editor/SequenceRecorder/Private/SequenceRecorder.h b/Engine/Source/Editor/SequenceRecorder/Private/SequenceRecorder.h
--- a/Engine/Source/Editor/SequenceRecorder/Private/SequenceRecorder.h
+++ b/Engine/Source/Editor/SequenceRecorder/Private/SequenceRecorder.h
@@ -174,7 +174,7 @@ private:
 	TWeakObjectPtr<class ULevelSequence> CurrentSequence;
 
 	/** World we are recording a replay for, if any. Is only valid if we're recording from Replay network drivers.*/
-	TLazyObjectPtr<class UWorld> CurrentReplayWorld;
+	TWeakObjectPtr<class UWorld> CurrentReplayWorld;
 
 	/** Actor World that our last started recording was. Null if there is not a recording in progress. */
 	TWeakObjectPtr<class UWorld> CurrentRecordingWorld;
diff --git a/Engine/Source/Editor/SequenceRecorder/Public/SequenceRecorderSettings.h b/Engine/Source/Editor/SequenceRecorder/Public/SequenceRecorderSettings.h
--- a/Engine/Source/Editor/SequenceRecorder/Public/SequenceRecorderSettings.h
+++ b/Engine/Source/Editor/SequenceRecorder/Public/SequenceRecorderSettings.h
@@ -172,7 +172,7 @@ public:
 
 	/** Sequence actors to trigger playback on when recording starts (e.g. for recording synchronized sequences) */
 	UPROPERTY(Transient, EditAnywhere, Category = "Sequence Recording")
-	TArray<TLazyObjectPtr<class ALevelSequenceActor>> LevelSequenceActorsToTrigger;
+	TArray<TWeakObjectPtr<class ALevelSequenceActor>> LevelSequenceActorsToTrigger;
 
 	/** Default animation settings which are used to initialize all new actor recording's animation settings */
 	UPROPERTY(Config, EditAnywhere, Category = "Sequence Recording")
diff --git a/Engine/Source/Editor/SequenceRecorderSections/Public/MovieScene3DAttachSectionRecorder.h b/Engine/Source/Editor/SequenceRecorderSections/Public/MovieScene3DAttachSectionRecorder.h
--- a/Engine/Source/Editor/SequenceRecorderSections/Public/MovieScene3DAttachSectionRecorder.h
+++ b/Engine/Source/Editor/SequenceRecorderSections/Public/MovieScene3DAttachSectionRecorder.h
@@ -10,7 +10,6 @@
 #include "MovieScene.h"
 #include "MovieSceneSection.h"
 #include "Templates/SharedPointer.h"
-#include "UObject/LazyObjectPtr.h"
 #include "UObject/WeakObjectPtr.h"
 #include "UObject/WeakObjectPtrTemplates.h"
 
@@ -46,7 +45,7 @@ public:
 
 private:
 	/** Object to record from */
-	TLazyObjectPtr<class AActor> ActorToRecord;
+	TWeakObjectPtr<class AActor> ActorToRecord;
 
 	/** Section to record to */
 	TWeakObjectPtr<class UMovieScene3DAttachSection> MovieSceneSection;
@@ -58,7 +57,7 @@ private:
 	TWeakObjectPtr<class UMovieScene> MovieScene;
 
 	/** Track the actor we are attached to */
-	TLazyObjectPtr<class AActor> ActorAttachedTo;
+	TWeakObjectPtr<class AActor> ActorAttachedTo;
 
 	/** Identifier of the object we are recording */
 	FGuid ObjectGuid;
diff --git a/Engine/Source/Editor/SequenceRecorderSections/Public/MovieSceneParticleTrackSectionRecorder.h b/Engine/Source/Editor/SequenceRecorderSections/Public/MovieSceneParticleTrackSectionRecorder.h
--- a/Engine/Source/Editor/SequenceRecorderSections/Public/MovieSceneParticleTrackSectionRecorder.h
+++ b/Engine/Source/Editor/SequenceRecorderSections/Public/MovieSceneParticleTrackSectionRecorder.h
@@ -70,7 +70,7 @@ public:
 
 private:
 	/** Object to record from */
-	TLazyObjectPtr<class UParticleSystemComponent> SystemToRecord;
+	TWeakObjectPtr<class UParticleSystemComponent> SystemToRecord;
 
 	/** Section to record to */
 	TWeakObjectPtr<class UMovieSceneParticleSection> MovieSceneSection;
diff --git a/Engine/Source/Editor/SequenceRecorderSections/Public/MovieSceneSpawnSectionRecorder.h b/Engine/Source/Editor/SequenceRecorderSections/Public/MovieSceneSpawnSectionRecorder.h
--- a/Engine/Source/Editor/SequenceRecorderSections/Public/MovieSceneSpawnSectionRecorder.h
+++ b/Engine/Source/Editor/SequenceRecorderSections/Public/MovieSceneSpawnSectionRecorder.h
@@ -6,7 +6,6 @@
 #include "IMovieSceneSectionRecorder.h"
 #include "IMovieSceneSectionRecorderFactory.h"
 #include "Templates/SharedPointer.h"
-#include "UObject/LazyObjectPtr.h"
 #include "UObject/Object.h"
 #include "UObject/WeakObjectPtr.h"
 #include "UObject/WeakObjectPtrTemplates.h"
@@ -42,7 +41,7 @@ public:
 
 private:
 	/** Object to record from */
-	TLazyObjectPtr<UObject> ObjectToRecord;
+	TWeakObjectPtr<UObject> ObjectToRecord;
 
 	/** Section to record to */
 	TWeakObjectPtr<class UMovieSceneBoolSection> MovieSceneSection;
diff --git a/Engine/Source/Editor/SequenceRecorderSections/Public/MovieSceneVisibilitySectionRecorder.h b/Engine/Source/Editor/SequenceRecorderSections/Public/MovieSceneVisibilitySectionRecorder.h
--- a/Engine/Source/Editor/SequenceRecorderSections/Public/MovieSceneVisibilitySectionRecorder.h
+++ b/Engine/Source/Editor/SequenceRecorderSections/Public/MovieSceneVisibilitySectionRecorder.h
@@ -9,7 +9,6 @@
 #include "IMovieSceneSectionRecorderFactory.h"
 #include "MovieSceneVisibilitySectionRecorderSettings.h"
 #include "Templates/SharedPointer.h"
-#include "UObject/LazyObjectPtr.h"
 #include "UObject/NameTypes.h"
 #include "UObject/Object.h"
 #include "UObject/UObjectGlobals.h"
@@ -48,7 +47,7 @@ public:
 
 private:
 	/** Object to record from */
-	TLazyObjectPtr<class UObject> ObjectToRecord;
+	TWeakObjectPtr<class UObject> ObjectToRecord;
 
 	/** Section to record to */
 	TWeakObjectPtr<class UMovieSceneBoolSection> MovieSceneSection;
diff --git a/Engine/Source/Runtime/LevelSequence/Private/LevelSequenceObject.cpp b/Engine/Source/Runtime/LevelSequence/Private/LevelSequenceObject.cpp
deleted file mode 100644
diff --git a/Engine/Source/Runtime/LevelSequence/Public/LevelSequence.h b/Engine/Source/Runtime/LevelSequence/Public/LevelSequence.h
--- a/Engine/Source/Runtime/LevelSequence/Public/LevelSequence.h
+++ b/Engine/Source/Runtime/LevelSequence/Public/LevelSequence.h
@@ -5,7 +5,6 @@
 #include "CoreMinimal.h"
 #include "UObject/ObjectMacros.h"
 #include "MovieSceneSequence.h"
-#include "LevelSequenceObject.h"
 #include "LevelSequenceBindingReference.h"
 #include "LevelSequenceLegacyObjectReference.h"
 #include "Templates/SubclassOf.h"
diff --git a/Engine/Source/Runtime/LevelSequence/Public/LevelSequenceLegacyObjectReference.h b/Engine/Source/Runtime/LevelSequence/Public/LevelSequenceLegacyObjectReference.h
--- a/Engine/Source/Runtime/LevelSequence/Public/LevelSequenceLegacyObjectReference.h
+++ b/Engine/Source/Runtime/LevelSequence/Public/LevelSequenceLegacyObjectReference.h
@@ -6,7 +6,6 @@
 #include "UObject/ObjectMacros.h"
 #include "Misc/Guid.h"
 #include "UObject/Class.h"
-#include "UObject/LazyObjectPtr.h"
 #include "LevelSequenceLegacyObjectReference.generated.h"
 
 /**
@@ -34,7 +33,7 @@ struct FLevelSequenceLegacyObjectReference
 		return Ar;
 	}
 
-	/** Primary method of resolution - object ID, stored as an annotation on the object in the world, resolvable through TLazyObjectPtr */
+	/** Primary method of resolution - object ID, stored as an annotation on the object in the world */
 	FUniqueObjectGuid ObjectId;
 
 	/** Secondary method of resolution - path to the object within the context */
diff --git a/Engine/Source/Runtime/LevelSequence/Public/LevelSequenceObject.h b/Engine/Source/Runtime/LevelSequence/Public/LevelSequenceObject.h
deleted file mode 100644
